#ue4 	CIS fix: Check for FSlateApplication validity.

diff --git a/Engine/Source/Developer/SourceCodeAccess/Private/SourceCodeAccessModule.cpp b/Engine/Source/Developer/SourceCodeAccess/Private/SourceCodeAccessModule.cpp
--- a/Engine/Source/Developer/SourceCodeAccess/Private/SourceCodeAccessModule.cpp
+++ b/Engine/Source/Developer/SourceCodeAccess/Private/SourceCodeAccessModule.cpp
@@ -78,8 +78,11 @@ void FSourceCodeAccessModule::SetAccessor(const FName& InName)
 		if(InName == Accessor.GetFName())
 		{
 			CurrentSourceCodeAccessor = static_cast<ISourceCodeAccessor*>(Feature);
-			FSlateApplication::Get().SetWidgetReflectorSourceAccessDelegate( FAccessSourceCode::CreateRaw( CurrentSourceCodeAccessor, &ISourceCodeAccessor::OpenFileAtLine ) );
-			FSlateApplication::Get().SetWidgetReflectorQuerySourceAccessDelegate( FQueryAccessSourceCode::CreateRaw( CurrentSourceCodeAccessor, &ISourceCodeAccessor::CanAccessSourceCode ) );
+			if(FSlateApplication::IsInitialized())
+			{
+				FSlateApplication::Get().SetWidgetReflectorSourceAccessDelegate( FAccessSourceCode::CreateRaw( CurrentSourceCodeAccessor, &ISourceCodeAccessor::OpenFileAtLine ) );
+				FSlateApplication::Get().SetWidgetReflectorQuerySourceAccessDelegate( FQueryAccessSourceCode::CreateRaw( CurrentSourceCodeAccessor, &ISourceCodeAccessor::CanAccessSourceCode ) );
+			}
 			break;
 		}
 	}
@@ -123,8 +126,11 @@ void FSourceCodeAccessModule::HandleModularFeatureRegistered(const FName& Type)
 			if(PreferredAccessor == Accessor.GetFName().ToString())
 			{
 				CurrentSourceCodeAccessor = static_cast<ISourceCodeAccessor*>(Feature);
-				FSlateApplication::Get().SetWidgetReflectorSourceAccessDelegate( FAccessSourceCode::CreateRaw( CurrentSourceCodeAccessor, &ISourceCodeAccessor::OpenFileAtLine ) );
-				FSlateApplication::Get().SetWidgetReflectorQuerySourceAccessDelegate( FQueryAccessSourceCode::CreateRaw( CurrentSourceCodeAccessor, &ISourceCodeAccessor::CanAccessSourceCode ) );
+				if(FSlateApplication::IsInitialized())
+				{
+					FSlateApplication::Get().SetWidgetReflectorSourceAccessDelegate( FAccessSourceCode::CreateRaw( CurrentSourceCodeAccessor, &ISourceCodeAccessor::OpenFileAtLine ) );
+					FSlateApplication::Get().SetWidgetReflectorQuerySourceAccessDelegate( FQueryAccessSourceCode::CreateRaw( CurrentSourceCodeAccessor, &ISourceCodeAccessor::CanAccessSourceCode ) );
+				}
 				break;
 			}
 		}
