Various fixes for using a non-default source code accessor in the editor.

* Fix incorrect path for tutorial dialog explaining how to install Visual Studio
* Tell the user that they need to restart when changing their source code accessor.
* On Windows, always check for the presence of Visual Studio 2017 or 2019 when determining whether the user can compile. This is distinct from whether the user has the chosen IDE available.

diff --git a/Engine/Source/Developer/SourceCodeAccess/Private/SourceCodeAccessModule.cpp b/Engine/Source/Developer/SourceCodeAccess/Private/SourceCodeAccessModule.cpp
--- a/Engine/Source/Developer/SourceCodeAccess/Private/SourceCodeAccessModule.cpp
+++ b/Engine/Source/Developer/SourceCodeAccess/Private/SourceCodeAccessModule.cpp
@@ -69,6 +69,28 @@ bool FSourceCodeAccessModule::CanAccessSourceCode() const
 	return CurrentSourceCodeAccessor->CanAccessSourceCode();
 }
 
+bool FSourceCodeAccessModule::CanCompileSourceCode() const
+{
+#if PLATFORM_WINDOWS
+	// Need to have Visual Studio installed to compile on Windows, regardless of chosen IDE
+	return IsSourceCodeAccessorAvailable("VisualStudio2017") || IsSourceCodeAccessorAvailable("VisualStudio2019");
+#else
+	// Default behavior
+	return CanAccessSourceCode();
+#endif
+}
+
+bool FSourceCodeAccessModule::IsSourceCodeAccessorAvailable(FName Name) const
+{
+	for (ISourceCodeAccessor* Accessor : IModularFeatures::Get().GetModularFeatureImplementations<ISourceCodeAccessor>(SourceCodeAccessorFeatureName))
+	{
+		if (Accessor->GetFName() == Name)
+		{
+			return true;
+		}
+	}
+	return false;
+}
 
 ISourceCodeAccessor& FSourceCodeAccessModule::GetAccessor() const
 {
diff --git a/Engine/Source/Developer/SourceCodeAccess/Private/SourceCodeAccessModule.h b/Engine/Source/Developer/SourceCodeAccess/Private/SourceCodeAccessModule.h
--- a/Engine/Source/Developer/SourceCodeAccess/Private/SourceCodeAccessModule.h
+++ b/Engine/Source/Developer/SourceCodeAccess/Private/SourceCodeAccessModule.h
@@ -29,6 +29,7 @@ public:
 	// ISourceCodeAccessModule interface
 
 	virtual bool CanAccessSourceCode() const override;
+	virtual bool CanCompileSourceCode() const override;
 	virtual ISourceCodeAccessor& GetAccessor() const override;
 	virtual void SetAccessor(const FName& InName) override;
 	virtual FLaunchingCodeAccessor& OnLaunchingCodeAccessor() override;
@@ -43,6 +44,9 @@ private:
 	/** Handle when one of the modular features we are interested in is unregistered */
 	void HandleModularFeatureUnregistered(const FName& Type, IModularFeature* ModularFeature);
 
+	/** Checks if a source code accessor with the given name is available */
+	bool IsSourceCodeAccessorAvailable(FName Name) const;
+
 private:
 
 	/** Event delegate fired when launching code accessor. */
diff --git a/Engine/Source/Developer/SourceCodeAccess/Private/SourceCodeAccessSettings.h b/Engine/Source/Developer/SourceCodeAccess/Private/SourceCodeAccessSettings.h
--- a/Engine/Source/Developer/SourceCodeAccess/Private/SourceCodeAccessSettings.h
+++ b/Engine/Source/Developer/SourceCodeAccess/Private/SourceCodeAccessSettings.h
@@ -13,6 +13,6 @@ class USourceCodeAccessSettings : public UObject
 	GENERATED_UCLASS_BODY()
 
 	/** The source code editor we prefer to use. */
-	UPROPERTY(Config, EditAnywhere, Category="Source Code Editor", meta=(DisplayName="Source Code Editor"))
+	UPROPERTY(Config, EditAnywhere, Category="Source Code Editor", meta=(DisplayName="Source Code Editor", ConfigRestartRequired = true))
 	FString PreferredAccessor;
 };
diff --git a/Engine/Source/Developer/SourceCodeAccess/Public/ISourceCodeAccessModule.h b/Engine/Source/Developer/SourceCodeAccess/Public/ISourceCodeAccessModule.h
--- a/Engine/Source/Developer/SourceCodeAccess/Public/ISourceCodeAccessModule.h
+++ b/Engine/Source/Developer/SourceCodeAccess/Public/ISourceCodeAccessModule.h
@@ -34,6 +34,12 @@ public:
 	 */
 	virtual bool CanAccessSourceCode() const = 0;
 
+	/** 
+	 * Checks if a compiler is available
+	 * @return true if source code can be compiled.
+	 */
+	virtual bool CanCompileSourceCode() const = 0;
+
 	/**
 	 * Get the accessor to allow us to view source code
 	 * @return the accessor.
diff --git a/Engine/Source/Editor/IntroTutorials/Private/IntroTutorials.cpp b/Engine/Source/Editor/IntroTutorials/Private/IntroTutorials.cpp
--- a/Engine/Source/Editor/IntroTutorials/Private/IntroTutorials.cpp
+++ b/Engine/Source/Editor/IntroTutorials/Private/IntroTutorials.cpp
@@ -361,9 +361,9 @@ void FIntroTutorials::OnNewProjectDialogOpened()
 void FIntroTutorials::HandleCompilerNotFound()
 {
 #if PLATFORM_WINDOWS
-	LaunchTutorialByName( TEXT( "Engine/Tutorial/Installation/InstallingVisualStudioTutorial.InstallingVisualStudioTutorial" ) );
+	LaunchTutorialByName( TEXT( "/Engine/Tutorial/Installation/InstallingVisualStudioTutorial.InstallingVisualStudioTutorial" ) );
 #elif PLATFORM_MAC
-	LaunchTutorialByName( TEXT( "Engine/Tutorial/Installation/InstallingXCodeTutorial.InstallingXCodeTutorial" ) );
+	LaunchTutorialByName( TEXT( "/Engine/Tutorial/Installation/InstallingXCodeTutorial.InstallingXCodeTutorial" ) );
 #else
 	STUBBED("FIntroTutorials::HandleCompilerNotFound");
 #endif
diff --git a/Engine/Source/Editor/UnrealEd/Private/SourceCodeNavigation.cpp b/Engine/Source/Editor/UnrealEd/Private/SourceCodeNavigation.cpp
--- a/Engine/Source/Editor/UnrealEd/Private/SourceCodeNavigation.cpp
+++ b/Engine/Source/Editor/UnrealEd/Private/SourceCodeNavigation.cpp
@@ -1726,7 +1726,7 @@ void FSourceCodeNavigation::RefreshCompilerAvailability()
 	ISourceCodeAccessModule& SourceCodeAccessModule = FModuleManager::LoadModuleChecked<ISourceCodeAccessModule>("SourceCodeAccess");
 	SourceCodeAccessModule.GetAccessor().RefreshAvailability();
 
-	bCachedIsCompilerAvailable = SourceCodeAccessModule.GetAccessor().CanAccessSourceCode();
+	bCachedIsCompilerAvailable = SourceCodeAccessModule.CanCompileSourceCode();
 }
 
 bool FSourceCodeNavigation::OpenSourceFile( const FString& AbsoluteSourcePath, int32 LineNumber, int32 ColumnNumber )
