Fix new project dialog allowing creation of new code projects even if an IDE is not installed.

#rnx

diff --git a/Engine/Source/Developer/SourceCodeAccess/Private/SourceCodeAccessModule.cpp b/Engine/Source/Developer/SourceCodeAccess/Private/SourceCodeAccessModule.cpp
--- a/Engine/Source/Developer/SourceCodeAccess/Private/SourceCodeAccessModule.cpp
+++ b/Engine/Source/Developer/SourceCodeAccess/Private/SourceCodeAccessModule.cpp
@@ -86,7 +86,7 @@ bool FSourceCodeAccessModule::IsSourceCodeAccessorAvailable(FName Name) const
 	{
 		if (Accessor->GetFName() == Name)
 		{
-			return true;
+			return Accessor->CanAccessSourceCode();
 		}
 	}
 	return false;
