Fix more missing copyright boilerplate in ISourceCodeAccessor.cpp that was missed by the last changelist

diff --git a/Engine/Source/Developer/SourceCodeAccess/Private/ISourceCodeAccessor.cpp b/Engine/Source/Developer/SourceCodeAccess/Private/ISourceCodeAccessor.cpp
--- a/Engine/Source/Developer/SourceCodeAccess/Private/ISourceCodeAccessor.cpp
+++ b/Engine/Source/Developer/SourceCodeAccess/Private/ISourceCodeAccessor.cpp
@@ -1,4 +1,6 @@
-﻿#include "ISourceCodeAccessor.h"
+﻿// Copyright Epic Games, Inc. All Rights Reserved.
+
+#include "ISourceCodeAccessor.h"
 
 #include "Styling/AppStyle.h"
 
